MeshConversion: add support for Tangents and Vertex Colors to FStaticMeshLODResourcesAdapter. Add new FStaticMeshLODResourcesToDynamicMesh conversion class.

diff --git a/Engine/Source/Runtime/MeshConversion/Public/StaticMeshLODResourcesAdapter.h b/Engine/Source/Runtime/MeshConversion/Public/StaticMeshLODResourcesAdapter.h
--- a/Engine/Source/Runtime/MeshConversion/Public/StaticMeshLODResourcesAdapter.h
+++ b/Engine/Source/Runtime/MeshConversion/Public/StaticMeshLODResourcesAdapter.h
@@ -5,6 +5,7 @@
 #include "CoreMinimal.h"
 #include "StaticMeshAttributes.h"
 #include "StaticMeshResources.h"
+#include "TransformTypes.h"
 
 
 /**
@@ -18,6 +19,7 @@ protected:
 	const FStaticMeshLODResources* Mesh;
 
 	FVector3d BuildScale = FVector3d::One();
+	FVector3d InvBuildScale = FVector3d::One();
 	bool bScaleNormals = false;
 
 	TArray<const FStaticMeshSection*> ValidSections;
@@ -52,6 +54,7 @@ public:
 	void SetBuildScale(const FVector3d& BuildScaleIn, bool bScaleNormalsIn)
 	{
 		BuildScale = BuildScaleIn;
+		InvBuildScale = UE::Geometry::FTransformSRT3d::GetSafeScaleReciprocal(BuildScaleIn);
 		bScaleNormals = bScaleNormalsIn;
 	}
 
@@ -149,7 +152,7 @@ public:
 	{
 		const FVector4f& Normal = Mesh->VertexBuffers.StaticMeshVertexBuffer.VertexTangentZ(IDValue);
 		return (!bScaleNormals) ? FVector3f(Normal.X, Normal.Y, Normal.Z) :
-			UE::Geometry::Normalized(FVector3f(Normal.X / BuildScale.X, Normal.Y / BuildScale.Y, Normal.Z / BuildScale.Z));
+			UE::Geometry::Normalized(FVector3f(Normal.X * InvBuildScale.X, Normal.Y * InvBuildScale.Y, Normal.Z * InvBuildScale.Z));
 	}
 
 	/** Get Normals for a given Triangle */
@@ -163,10 +166,52 @@ public:
 		N2 = GetNormal(VtxIndices.C);
 	}
 
+
+	inline FVector3f GetTangentX(int32 IDValue) const
+	{
+		const FVector4f& TangentX = Mesh->VertexBuffers.StaticMeshVertexBuffer.VertexTangentX(IDValue);
+		return (!bScaleNormals) ? FVector3f(TangentX.X, TangentX.Y, TangentX.Z) :
+			UE::Geometry::Normalized(FVector3f(TangentX.X * BuildScale.X, TangentX.Y * BuildScale.Y, TangentX.Z * BuildScale.Z));
+	}
+
+	/** Get Tangent X for a given Triangle */
+	template<typename VectorType>
+	inline void GetTriTangentsX(int32 TriId, VectorType& T0, VectorType& T1, VectorType& T2)
+	{
+		FIndex3i VtxIndices = GetTriangle(TriId);
+		T0 = GetTangentX(VtxIndices.A);
+		T1 = GetTangentX(VtxIndices.B);
+		T2 = GetTangentX(VtxIndices.C);
+	}
+
+
+	inline FVector3f GetTangentY(int32 IDValue) const
+	{
+		const FVector4f& TangentY = Mesh->VertexBuffers.StaticMeshVertexBuffer.VertexTangentY(IDValue);
+		return (!bScaleNormals) ? FVector3f(TangentY.X, TangentY.Y, TangentY.Z) :
+			UE::Geometry::Normalized(FVector3f(TangentY.X * BuildScale.X, TangentY.Y * BuildScale.Y, TangentY.Z * BuildScale.Z));
+	}
+
+	/** Get Tangent Y for a given Triangle */
+	template<typename VectorType>
+	inline void GetTriTangentsY(int32 TriId, VectorType& T0, VectorType& T1, VectorType& T2)
+	{
+		FIndex3i VtxIndices = GetTriangle(TriId);
+		T0 = GetTangentY(VtxIndices.A);
+		T1 = GetTangentY(VtxIndices.B);
+		T2 = GetTangentY(VtxIndices.C);
+	}
+
+
 	inline bool HasUVs(const int32 UVLayer = 0) const
 	{
 		return Mesh && Mesh->VertexBuffers.StaticMeshVertexBuffer.GetAllowCPUAccess() && Mesh->VertexBuffers.StaticMeshVertexBuffer.GetTexCoordData() && (int32)Mesh->VertexBuffers.StaticMeshVertexBuffer.GetNumTexCoords() > UVLayer && UVLayer >= 0;
 	}
+	inline int32 NumUVLayers() const
+	{
+		return (Mesh && Mesh->VertexBuffers.StaticMeshVertexBuffer.GetAllowCPUAccess() && Mesh->VertexBuffers.StaticMeshVertexBuffer.GetTexCoordData()) ?
+			(int32)Mesh->VertexBuffers.StaticMeshVertexBuffer.GetNumTexCoords() : 0;
+	}
 	inline bool IsUV(const int32 UVId) const
 	{
 		return HasUVs() && UVId >= 0 && UVId < UVCount();
@@ -196,4 +241,37 @@ public:
 		UV1 = GetUV(VtxIndices.B, UVLayer);
 		UV2 = GetUV(VtxIndices.C, UVLayer);
 	}
+
+
+	inline bool HasColors() const
+	{
+		return Mesh && Mesh->VertexBuffers.ColorVertexBuffer.GetAllowCPUAccess();
+	}
+	inline bool IsColor(int32 ColorIndex) const
+	{
+		return ColorIndex >= 0 && ColorIndex < ColorCount();
+	}
+	inline int32 MaxColorID() const
+	{
+		return ColorCount();
+	}
+	inline int32 ColorCount() const
+	{
+		return HasColors() ? VertexCount() : 0;
+	}
+
+	inline FColor GetColor(int32 IDValue) const
+	{
+		return Mesh->VertexBuffers.ColorVertexBuffer.VertexColor(IDValue);
+	}
+
+	/** Get Colors for a given Triangle */
+	inline void GetTriColors(int32 TriId, FColor& C0, FColor& C1, FColor& C2)
+	{
+		FIndex3i VtxIndices = GetTriangle(TriId);
+		C0 = GetColor(VtxIndices.A);
+		C1 = GetColor(VtxIndices.B);
+		C2 = GetColor(VtxIndices.C);
+	}
+
 };
diff --git a/Engine/Source/Runtime/MeshConversionEngineTypes/Private/StaticMeshLODResourcesToDynamicMesh.cpp b/Engine/Source/Runtime/MeshConversionEngineTypes/Private/StaticMeshLODResourcesToDynamicMesh.cpp
new file mode 100644
--- /dev/null
+++ b/Engine/Source/Runtime/MeshConversionEngineTypes/Private/StaticMeshLODResourcesToDynamicMesh.cpp
@@ -0,0 +1,175 @@
+// Copyright Epic Games, Inc. All Rights Reserved.
+
+#include "StaticMeshLODResourcesToDynamicMesh.h"
+#include "StaticMeshLODResourcesAdapter.h"
+#include "DynamicMesh/DynamicMeshAttributeSet.h"
+
+#include "StaticMeshAttributes.h"
+#include "StaticMeshResources.h"
+
+
+using namespace UE::Geometry;
+
+bool FStaticMeshLODResourcesToDynamicMesh::Convert(
+	const FStaticMeshLODResources* StaticMeshResources,
+	const ConversionOptions& Options,
+	FDynamicMesh3& OutputMesh)
+{
+	FStaticMeshLODResourcesMeshAdapter Adapter(StaticMeshResources);
+
+	Adapter.SetBuildScale(Options.BuildScale, false);
+
+	OutputMesh = FDynamicMesh3();
+	OutputMesh.EnableTriangleGroups();
+	if (Options.bWantNormals || Options.bWantTangents || Options.bWantUVs || Options.bWantVertexColors || Options.bWantMaterialIDs)
+	{
+		OutputMesh.EnableAttributes();
+	}
+
+	// Copy vertices. LODMesh is dense so this should be 1-1
+	int32 VertexCount = Adapter.VertexCount();
+	for ( int32 VertID = 0; VertID < VertexCount; ++VertID )
+	{
+		FVector3d Position = Adapter.GetVertex(VertID);
+		int NewVertID = OutputMesh.AppendVertex(Position);
+		if (NewVertID != VertID)
+		{
+			OutputMesh.Clear();
+			ensure(false);
+			return false;
+		}
+	}
+
+	// Copy triangles. LODMesh is dense so this should be 1-1 unless there is a duplicate tri or non-manifold edge (currently aborting in that case)
+	int32 TriangleCount = Adapter.TriangleCount();
+	for (int32 TriID = 0; TriID < TriangleCount; ++TriID)
+	{
+		FIndex3i Tri = Adapter.GetTriangle(TriID);
+		int32 NewTriID = OutputMesh.AppendTriangle(Tri.A, Tri.B, Tri.C);
+		if (NewTriID != TriID)
+		{
+			OutputMesh.Clear();
+			ensure(false);
+			return false;
+		}
+
+	}
+
+	// transfer sections to PolyGroups and MaterialIDs
+	FDynamicMeshMaterialAttribute* MaterialIDs = nullptr;
+	if (Options.bWantMaterialIDs)
+	{
+		OutputMesh.Attributes()->EnableMaterialID();
+		MaterialIDs = OutputMesh.Attributes()->GetMaterialID();
+	}
+	for ( int32 SectionIdx = 0; SectionIdx < StaticMeshResources->Sections.Num(); ++SectionIdx)
+	{
+		const FStaticMeshSection& Section = StaticMeshResources->Sections[SectionIdx];
+		for (uint32 TriIdx = 0; TriIdx < Section.NumTriangles; ++TriIdx)
+		{
+			OutputMesh.SetTriangleGroup((int32)(Section.FirstIndex + TriIdx), SectionIdx);
+			if (MaterialIDs != nullptr)
+			{
+				MaterialIDs->SetValue((int32)(Section.FirstIndex + TriIdx), Section.MaterialIndex);
+			}
+		}
+	}
+
+	// copy overlay normals
+	if (Adapter.HasNormals() && Options.bWantNormals)
+	{
+
+
+		FDynamicMeshNormalOverlay* Normals = OutputMesh.Attributes()->PrimaryNormals();
+		if (Normals != nullptr)
+		{
+			for (int32 VertID = 0; VertID < VertexCount; ++VertID)
+			{
+				FVector3f N = Adapter.GetNormal(VertID);
+				int32 ElemID = Normals->AppendElement(N);
+				check(ElemID == VertID);
+			}
+
+			for (int32 TriID = 0; TriID < TriangleCount; ++TriID)
+			{
+				FIndex3i Tri = Adapter.GetTriangle(TriID);
+				Normals->SetTriangle(TriID, FIndex3i(Tri.A, Tri.B, Tri.C));
+			}
+		}
+	}
+
+	// copy overlay tangents
+	if ( Adapter.HasNormals() && Options.bWantTangents )
+	{
+		OutputMesh.Attributes()->EnableTangents();
+		FDynamicMeshNormalOverlay* TangentsX = OutputMesh.Attributes()->PrimaryTangents();
+		if (TangentsX != nullptr)
+		{
+			for (int32 TriID = 0; TriID < TriangleCount; ++TriID)
+			{
+				FVector3f T1, T2, T3;
+				Adapter.GetTriTangentsX<FVector3f>(TriID, T1, T2, T3);
+				int32 a = TangentsX->AppendElement(T1);
+				int32 b = TangentsX->AppendElement(T2);
+				int32 c = TangentsX->AppendElement(T3);
+				TangentsX->SetTriangle(TriID, FIndex3i(a, b, c));
+			}
+		}
+
+		FDynamicMeshNormalOverlay* TangentsY = OutputMesh.Attributes()->PrimaryBiTangents();
+		if (TangentsY != nullptr)
+		{
+			for (int32 TriID = 0; TriID < TriangleCount; ++TriID)
+			{
+				FVector3f T1, T2, T3;
+				Adapter.GetTriTangentsY<FVector3f>(TriID, T1, T2, T3);
+				int32 a = TangentsY->AppendElement(T1);
+				int32 b = TangentsY->AppendElement(T2);
+				int32 c = TangentsY->AppendElement(T3);
+				TangentsY->SetTriangle(TriID, FIndex3i(a, b, c));
+			}
+		}
+	}
+
+	// copy UV layers
+	if (Adapter.HasUVs() && Options.bWantUVs)
+	{
+		int32 NumUVLayers = Adapter.NumUVLayers();
+		if (NumUVLayers > 0)
+		{
+			OutputMesh.Attributes()->SetNumUVLayers(NumUVLayers);
+			for (int32 UVLayerIndex = 0; UVLayerIndex < NumUVLayers; ++UVLayerIndex)
+			{
+				FDynamicMeshUVOverlay* UVOverlay = OutputMesh.Attributes()->GetUVLayer(UVLayerIndex);
+				for (int32 TriID = 0; TriID < TriangleCount; ++TriID)
+				{
+					FVector2f UV1, UV2, UV3;
+					Adapter.GetTriUVs<FVector2f>(TriID, UVLayerIndex, UV1, UV2, UV3);
+					int32 a = UVOverlay->AppendElement(UV1);
+					int32 b = UVOverlay->AppendElement(UV2);
+					int32 c = UVOverlay->AppendElement(UV3);
+					UVOverlay->SetTriangle(TriID, FIndex3i(a, b, c));
+				}
+			}
+		}
+	}
+
+	// copy overlay colors
+	if ( Adapter.HasColors() && Options.bWantVertexColors )
+	{
+		OutputMesh.Attributes()->EnablePrimaryColors();
+		FDynamicMeshColorOverlay* Colors = OutputMesh.Attributes()->PrimaryColors();
+		for (int32 TriID = 0; TriID < TriangleCount; ++TriID)
+		{
+			FColor C1, C2, C3;
+			Adapter.GetTriColors(TriID, C1, C2, C3);
+			// [TODO] should we be doing RGBA conversion here?
+			int32 a = Colors->AppendElement( C1.ReinterpretAsLinear()  );
+			int32 b = Colors->AppendElement( C2.ReinterpretAsLinear() );
+			int32 c = Colors->AppendElement( C3.ReinterpretAsLinear() );
+			Colors->SetTriangle(TriID, FIndex3i(a, b, c));
+		}
+	}
+
+	return true;
+}
\ No newline at end of file
diff --git a/Engine/Source/Runtime/MeshConversionEngineTypes/Public/StaticMeshLODResourcesToDynamicMesh.h b/Engine/Source/Runtime/MeshConversionEngineTypes/Public/StaticMeshLODResourcesToDynamicMesh.h
new file mode 100644
--- /dev/null
+++ b/Engine/Source/Runtime/MeshConversionEngineTypes/Public/StaticMeshLODResourcesToDynamicMesh.h
@@ -0,0 +1,47 @@
+// Copyright Epic Games, Inc. All Rights Reserved.
+
+#pragma once
+
+#include "CoreMinimal.h"
+#include "DynamicMesh/DynamicMesh3.h"
+
+struct FStaticMeshLODResources;
+
+namespace UE
+{
+namespace Geometry
+{
+
+
+/**
+ * FStaticMeshLODResourcesToDynamicMesh can be used to create a FDynamicMesh3 from
+ * a FStaticMeshLODResources, which is the variant of the mesh in a StaticMesh Asset used for rendering
+ * (and which is available at runtime). The FStaticMeshLODResources has vertices duplicated at any
+ * split UV/normal/tangent/color, ie in the overlays there will be a unique overlay element for each
+ * base mesh vertex.
+ */
+class MESHCONVERSIONENGINETYPES_API FStaticMeshLODResourcesToDynamicMesh
+{
+public:
+
+	struct MESHCONVERSIONENGINETYPES_API ConversionOptions
+	{
+		bool bWantNormals = true;
+		bool bWantTangents = true;
+		bool bWantUVs = true;
+		bool bWantVertexColors = true;
+		bool bWantMaterialIDs = true;
+
+		// mesh vertex positions are multiplied by the build scale
+		FVector3d BuildScale = FVector3d::One();
+	};
+
+	static bool Convert(
+		const FStaticMeshLODResources* StaticMeshResources, 
+		const ConversionOptions& Options, 
+		FDynamicMesh3& OutputMesh);
+
+};
+
+}
+}
