Fix indexing bug in FStaticMeshLODResourcesToDynamicMesh::Convert. Remove unnecessary template type on TImageBuilder::SampleNearestUV

diff --git a/Engine/Plugins/Experimental/GeometryScripting/Source/GeometryScriptingCore/Private/TextureMapFunctions.cpp b/Engine/Plugins/Experimental/GeometryScripting/Source/GeometryScriptingCore/Private/TextureMapFunctions.cpp
--- a/Engine/Plugins/Experimental/GeometryScripting/Source/GeometryScriptingCore/Private/TextureMapFunctions.cpp
+++ b/Engine/Plugins/Experimental/GeometryScripting/Source/GeometryScriptingCore/Private/TextureMapFunctions.cpp
@@ -54,7 +54,7 @@ void UGeometryScriptLibrary_TextureMapFunctions::SampleTexture2DAtUVPositions(
 		}
 
 		FVector4f InterpValue = (SampleOptions.SamplingMethod == EGeometryScriptPixelSamplingMethod::Bilinear) ?
-			ImageData.BilinearSampleUV<double>(UV, FVector4f::Zero()) : ImageData.NearestSampleUV<double>(UV, FVector4f::Zero());
+			ImageData.BilinearSampleUV<double>(UV, FVector4f::Zero()) : ImageData.NearestSampleUV(UV, FVector4f::Zero());
 
 		Colors[k] = (FLinearColor)InterpValue;
 	}
diff --git a/Engine/Source/Runtime/GeometryCore/Public/Image/ImageBuilder.h b/Engine/Source/Runtime/GeometryCore/Public/Image/ImageBuilder.h
--- a/Engine/Source/Runtime/GeometryCore/Public/Image/ImageBuilder.h
+++ b/Engine/Source/Runtime/GeometryCore/Public/Image/ImageBuilder.h
@@ -221,7 +221,6 @@ public:
 	* Sample the image value at floating-point UV coords with Nearest sampling.
 	* The UV coords are assumed to be in range [0,1]x[0,1], and that this maps to the [0,Width]x[0,Height] image pixel rectangle.
 	*/
-	template<typename ScalarType>
 	PixelType NearestSampleUV(const FVector2d& UVCoords, const PixelType& InvalidValue) const
 	{
 		int X = (int)(UVCoords.X * (double)Dimensions.GetWidth());
diff --git a/Engine/Source/Runtime/MeshConversionEngineTypes/Private/StaticMeshLODResourcesToDynamicMesh.cpp b/Engine/Source/Runtime/MeshConversionEngineTypes/Private/StaticMeshLODResourcesToDynamicMesh.cpp
--- a/Engine/Source/Runtime/MeshConversionEngineTypes/Private/StaticMeshLODResourcesToDynamicMesh.cpp
+++ b/Engine/Source/Runtime/MeshConversionEngineTypes/Private/StaticMeshLODResourcesToDynamicMesh.cpp
@@ -67,10 +67,11 @@ bool FStaticMeshLODResourcesToDynamicMesh::Convert(
 		const FStaticMeshSection& Section = StaticMeshResources->Sections[SectionIdx];
 		for (uint32 TriIdx = 0; TriIdx < Section.NumTriangles; ++TriIdx)
 		{
-			OutputMesh.SetTriangleGroup((int32)(Section.FirstIndex + TriIdx), SectionIdx);
+			int32 TriangleID = (int32)(Section.FirstIndex/3 + TriIdx);
+			OutputMesh.SetTriangleGroup(TriangleID, SectionIdx);
 			if (MaterialIDs != nullptr)
 			{
-				MaterialIDs->SetValue((int32)(Section.FirstIndex + TriIdx), Section.MaterialIndex);
+				MaterialIDs->SetValue(TriangleID, Section.MaterialIndex);
 			}
 		}
 	}
@@ -78,8 +79,6 @@ bool FStaticMeshLODResourcesToDynamicMesh::Convert(
 	// copy overlay normals
 	if (Adapter.HasNormals() && Options.bWantNormals)
 	{
-
-
 		FDynamicMeshNormalOverlay* Normals = OutputMesh.Attributes()->PrimaryNormals();
 		if (Normals != nullptr)
 		{
@@ -172,4 +171,4 @@ bool FStaticMeshLODResourcesToDynamicMesh::Convert(
 	}
 
 	return true;
-}
\ No newline at end of file
+}
