Guard against runtime StaticMesh -> DynamicMesh conversion functions crashing when CPU access to render data is not available; instead ensureMsgf() with a message explaining that CPU access is required

diff --git a/Engine/Source/Runtime/MeshConversionEngineTypes/Private/IndexedMeshToDynamicMesh.cpp b/Engine/Source/Runtime/MeshConversionEngineTypes/Private/IndexedMeshToDynamicMesh.cpp
--- a/Engine/Source/Runtime/MeshConversionEngineTypes/Private/IndexedMeshToDynamicMesh.cpp
+++ b/Engine/Source/Runtime/MeshConversionEngineTypes/Private/IndexedMeshToDynamicMesh.cpp
@@ -8,13 +8,17 @@
 
 using namespace UE::Geometry;
 
-void UE::Conversion::RenderBuffersToDynamicMesh(
+bool UE::Conversion::RenderBuffersToDynamicMesh(
 	const FStaticMeshVertexBuffers& VertexData,
 	const FRawStaticIndexBuffer& IndexData,
 	const FStaticMeshSectionArray& SectionData,
 	FDynamicMesh3& MeshOut,
 	bool bAttemptToWeldSeams)
 {
+	if (!ensureMsgf(VertexData.StaticMeshVertexBuffer.GetAllowCPUAccess(), TEXT("bAllowCPUAccess must be set to true for StaticMeshes before calling RenderBuffersToDynamicMesh(), otherwise the mesh geometry data isn't accessible!")))
+	{
+		return false;
+	}
 	const FStaticMeshVertexBuffer& VertexBuffer = VertexData.StaticMeshVertexBuffer;
 	const FPositionVertexBuffer& PositionBuffer = VertexData.PositionVertexBuffer;
 	const FColorVertexBuffer& ColorBuffer = VertexData.ColorVertexBuffer;
@@ -179,4 +183,6 @@ void UE::Conversion::RenderBuffersToDynamicMesh(
 		FMergeCoincidentMeshEdges Merge(&MeshOut);
 		Merge.Apply();
 	}
+
+	return true;
 }
\ No newline at end of file
diff --git a/Engine/Source/Runtime/MeshConversionEngineTypes/Private/StaticMeshLODResourcesToDynamicMesh.cpp b/Engine/Source/Runtime/MeshConversionEngineTypes/Private/StaticMeshLODResourcesToDynamicMesh.cpp
--- a/Engine/Source/Runtime/MeshConversionEngineTypes/Private/StaticMeshLODResourcesToDynamicMesh.cpp
+++ b/Engine/Source/Runtime/MeshConversionEngineTypes/Private/StaticMeshLODResourcesToDynamicMesh.cpp
@@ -15,6 +15,11 @@ bool FStaticMeshLODResourcesToDynamicMesh::Convert(
 	const ConversionOptions& Options,
 	FDynamicMesh3& OutputMesh)
 {
+	if (!ensureMsgf(StaticMeshResources && StaticMeshResources->VertexBuffers.StaticMeshVertexBuffer.GetAllowCPUAccess(), TEXT("bAllowCPUAccess must be set to true for StaticMeshes before calling FStaticMeshLODResourcesToDynamicMesh::Convert(), otherwise the mesh geometry data isn't accessible!")))
+	{
+		return false;
+	}
+
 	FStaticMeshLODResourcesMeshAdapter Adapter(StaticMeshResources);
 
 	Adapter.SetBuildScale(Options.BuildScale, false);
diff --git a/Engine/Source/Runtime/MeshConversionEngineTypes/Public/IndexedMeshToDynamicMesh.h b/Engine/Source/Runtime/MeshConversionEngineTypes/Public/IndexedMeshToDynamicMesh.h
--- a/Engine/Source/Runtime/MeshConversionEngineTypes/Public/IndexedMeshToDynamicMesh.h
+++ b/Engine/Source/Runtime/MeshConversionEngineTypes/Public/IndexedMeshToDynamicMesh.h
@@ -18,7 +18,7 @@ using namespace UE::Geometry;
  * as well as overlay Tangents enabled/copied, and MaterialID set.
  * @param bAttemptToWeldSeams if true, try to weld seam edges together to create a more closed mesh using FMergeCoincidentMeshEdges
  */
-void MESHCONVERSIONENGINETYPES_API RenderBuffersToDynamicMesh(
+bool MESHCONVERSIONENGINETYPES_API RenderBuffersToDynamicMesh(
 	const FStaticMeshVertexBuffers& VertexData,
 	const FRawStaticIndexBuffer& IndexData,
 	const FStaticMeshSectionArray& SectionData,
