Make conversion from render data to dynamic mesh more robust to non-manifold edges.

diff --git a/Engine/Source/Runtime/MeshConversionEngineTypes/Private/StaticMeshLODResourcesToDynamicMesh.cpp b/Engine/Source/Runtime/MeshConversionEngineTypes/Private/StaticMeshLODResourcesToDynamicMesh.cpp
--- a/Engine/Source/Runtime/MeshConversionEngineTypes/Private/StaticMeshLODResourcesToDynamicMesh.cpp
+++ b/Engine/Source/Runtime/MeshConversionEngineTypes/Private/StaticMeshLODResourcesToDynamicMesh.cpp
@@ -31,33 +31,90 @@ bool FStaticMeshLODResourcesToDynamicMesh::Convert(
 		OutputMesh.EnableAttributes();
 	}
 
+	// map from StaticMeshResouce triangle ID to DynamicMesh TID
+	TArray<int32> ToDstTriID;
+
+	/**
+	* map from DynamicMesh vertex Id to StaticMeshResouce vertex id.
+	* NB: due to vertex splitting, multiple DynamicMesh vertex ids
+	* may map to the same StaticMeshResource vertex id.
+	*  ( a vertex split is a result of reconciling non-manifold MeshDescription vertex )
+	*/
+	TArray<int32> ToSrcVID;
+
 	// Copy vertices. LODMesh is dense so this should be 1-1
-	int32 VertexCount = Adapter.VertexCount();
-	for ( int32 VertID = 0; VertID < VertexCount; ++VertID )
+	int32 SrcVertexCount = Adapter.VertexCount();
+	ToSrcVID.SetNumUninitialized(SrcVertexCount);
+	for ( int32 SrcVertID = 0; SrcVertID < SrcVertexCount; ++SrcVertID )
 	{
-		FVector3d Position = Adapter.GetVertex(VertID);
-		int NewVertID = OutputMesh.AppendVertex(Position);
-		if (NewVertID != VertID)
+		FVector3d Position = Adapter.GetVertex(SrcVertID);
+		int32 DstVertID = OutputMesh.AppendVertex(Position);
+		ToSrcVID[DstVertID] = SrcVertID;
+		
+		if (DstVertID != SrcVertID)
 		{
+			// should only happen in the source mesh is missing vertices. 
 			OutputMesh.Clear();
 			ensure(false);
 			return false;
 		}
 	}
 
-	// Copy triangles. LODMesh is dense so this should be 1-1 unless there is a duplicate tri or non-manifold edge (currently aborting in that case)
-	int32 TriangleCount = Adapter.TriangleCount();
-	for (int32 TriID = 0; TriID < TriangleCount; ++TriID)
+	// Copy triangles. LODMesh is dense so this should be 1-1 unless there is a duplicate tri or non-manifold edge 
+	int32 SrcTriangleCount = Adapter.TriangleCount();
+	ToDstTriID.Init(FDynamicMesh3::InvalidID, SrcTriangleCount);
+	for (int32 SrcTriID = 0; SrcTriID < SrcTriangleCount; ++SrcTriID)
 	{
-		FIndex3i Tri = Adapter.GetTriangle(TriID);
-		int32 NewTriID = OutputMesh.AppendTriangle(Tri.A, Tri.B, Tri.C);
-		if (NewTriID != TriID)
+		FIndex3i Tri = Adapter.GetTriangle(SrcTriID);
+		int32 DstTriID = OutputMesh.AppendTriangle(Tri.A, Tri.B, Tri.C);
+
+		if (DstTriID == FDynamicMesh3::DuplicateTriangleID || DstTriID == FDynamicMesh3::InvalidID)
 		{
-			OutputMesh.Clear();
-			ensure(false);
-			return false;
+			continue;
+		}
+
+		// split verts on the non-manifold edge(s)
+		if (DstTriID == FDynamicMesh3::NonManifoldID)
+		{
+			int e01 = OutputMesh.FindEdge(Tri[0], Tri[1]);
+			int e12 = OutputMesh.FindEdge(Tri[1], Tri[2]);
+			int e20 = OutputMesh.FindEdge(Tri[2], Tri[0]);
+
+			// determine which verts need to be duplicated
+			bool bToSplit[3] = { false, false, false };
+			if (e01 != FDynamicMesh3::InvalidID && OutputMesh.IsBoundaryEdge(e01) == false)
+			{
+				bToSplit[0] = true;
+				bToSplit[1] = true;
+			}
+			if (e12 != FDynamicMesh3::InvalidID && OutputMesh.IsBoundaryEdge(e12) == false)
+			{
+				bToSplit[1] = true;
+				bToSplit[2] = true;
+			}
+			if (e20 != FDynamicMesh3::InvalidID && OutputMesh.IsBoundaryEdge(e20) == false)
+			{
+				bToSplit[2] = true;
+				bToSplit[0] = true;
+			}
+			for (int32 i = 0; i < 3; ++i)
+			{
+				if (bToSplit[i])
+				{
+					const int32 SrcVID = Tri[i];
+					const FVector3d Position = Adapter.GetVertex(SrcVID);
+					const int32 NewDstVertIdx = OutputMesh.AppendVertex(Position);
+					Tri[i] = NewDstVertIdx;
+					ToSrcVID.SetNumUninitialized(NewDstVertIdx + 1);
+					ToSrcVID[NewDstVertIdx] = SrcVID;
+				}
+			}
+
+			DstTriID = OutputMesh.AppendTriangle(Tri[0], Tri[1], Tri[2]);
 		}
 
+		ToDstTriID[SrcTriID] = DstTriID;
+
 	}
 
 	// transfer sections to PolyGroups and MaterialIDs
@@ -72,11 +129,15 @@ bool FStaticMeshLODResourcesToDynamicMesh::Convert(
 		const FStaticMeshSection& Section = StaticMeshResources->Sections[SectionIdx];
 		for (uint32 TriIdx = 0; TriIdx < Section.NumTriangles; ++TriIdx)
 		{
-			int32 TriangleID = (int32)(Section.FirstIndex/3 + TriIdx);
-			OutputMesh.SetTriangleGroup(TriangleID, SectionIdx);
-			if (MaterialIDs != nullptr)
-			{
-				MaterialIDs->SetValue(TriangleID, Section.MaterialIndex);
+			int32 SrcTriangleID = (int32)(Section.FirstIndex/3 + TriIdx);
+			int32 DstTriID = ToDstTriID[SrcTriangleID];
+			if (DstTriID != FDynamicMesh3::InvalidID)
+			{ 
+				OutputMesh.SetTriangleGroup(DstTriID, SectionIdx);
+				if (MaterialIDs != nullptr)
+				{
+					MaterialIDs->SetValue(DstTriID, Section.MaterialIndex);
+				}
 			}
 		}
 	}
@@ -87,17 +148,23 @@ bool FStaticMeshLODResourcesToDynamicMesh::Convert(
 		FDynamicMeshNormalOverlay* Normals = OutputMesh.Attributes()->PrimaryNormals();
 		if (Normals != nullptr)
 		{
-			for (int32 VertID = 0; VertID < VertexCount; ++VertID)
+			const int32 DstVertexCount = ToSrcVID.Num();
+			for (int32 DstVertID = 0; DstVertID < DstVertexCount; ++DstVertID)
 			{
-				FVector3f N = Adapter.GetNormal(VertID);
+				const int32 SrcVID = ToSrcVID[DstVertID];
+				FVector3f N = Adapter.GetNormal(SrcVID);
 				int32 ElemID = Normals->AppendElement(N);
-				check(ElemID == VertID);
+				check(ElemID == DstVertID);
 			}
 
-			for (int32 TriID = 0; TriID < TriangleCount; ++TriID)
+			for (int32 SrcTriID = 0; SrcTriID < SrcTriangleCount; ++SrcTriID)
 			{
-				FIndex3i Tri = Adapter.GetTriangle(TriID);
-				Normals->SetTriangle(TriID, FIndex3i(Tri.A, Tri.B, Tri.C));
+				const int32 DstTriID = ToDstTriID[SrcTriID];
+				if (DstTriID != FDynamicMesh3::InvalidID)
+				{ 
+					FIndex3i Tri = OutputMesh.GetTriangle(DstTriID);
+					Normals->SetTriangle(DstTriID, FIndex3i(Tri.A, Tri.B, Tri.C));
+				}
 			}
 		}
 	}
@@ -109,28 +176,36 @@ bool FStaticMeshLODResourcesToDynamicMesh::Convert(
 		FDynamicMeshNormalOverlay* TangentsX = OutputMesh.Attributes()->PrimaryTangents();
 		if (TangentsX != nullptr)
 		{
-			for (int32 TriID = 0; TriID < TriangleCount; ++TriID)
+			for (int32 SrcTriID = 0; SrcTriID < SrcTriangleCount; ++SrcTriID)
 			{
-				FVector3f T1, T2, T3;
-				Adapter.GetTriTangentsX<FVector3f>(TriID, T1, T2, T3);
-				int32 a = TangentsX->AppendElement(T1);
-				int32 b = TangentsX->AppendElement(T2);
-				int32 c = TangentsX->AppendElement(T3);
-				TangentsX->SetTriangle(TriID, FIndex3i(a, b, c));
+				const int32 DstTriID = ToDstTriID[SrcTriID];
+				if (DstTriID != FDynamicMesh3::InvalidID)
+				{ 
+					FVector3f T1, T2, T3;
+					Adapter.GetTriTangentsX<FVector3f>(SrcTriID, T1, T2, T3);
+					int32 a = TangentsX->AppendElement(T1);
+					int32 b = TangentsX->AppendElement(T2);
+					int32 c = TangentsX->AppendElement(T3);
+					TangentsX->SetTriangle(DstTriID, FIndex3i(a, b, c));
+				}
 			}
 		}
 
 		FDynamicMeshNormalOverlay* TangentsY = OutputMesh.Attributes()->PrimaryBiTangents();
 		if (TangentsY != nullptr)
 		{
-			for (int32 TriID = 0; TriID < TriangleCount; ++TriID)
+			for (int32 SrcTriID = 0; SrcTriID < SrcTriangleCount; ++SrcTriID)
 			{
-				FVector3f T1, T2, T3;
-				Adapter.GetTriTangentsY<FVector3f>(TriID, T1, T2, T3);
-				int32 a = TangentsY->AppendElement(T1);
-				int32 b = TangentsY->AppendElement(T2);
-				int32 c = TangentsY->AppendElement(T3);
-				TangentsY->SetTriangle(TriID, FIndex3i(a, b, c));
+				const int32 DstTriID = ToDstTriID[SrcTriID];
+				if (DstTriID != FDynamicMesh3::InvalidID)
+				{ 
+					FVector3f T1, T2, T3;
+					Adapter.GetTriTangentsY<FVector3f>(SrcTriID, T1, T2, T3);
+					int32 a = TangentsY->AppendElement(T1);
+					int32 b = TangentsY->AppendElement(T2);
+					int32 c = TangentsY->AppendElement(T3);
+					TangentsY->SetTriangle(DstTriID, FIndex3i(a, b, c));
+				}
 			}
 		}
 	}
@@ -145,14 +220,18 @@ bool FStaticMeshLODResourcesToDynamicMesh::Convert(
 			for (int32 UVLayerIndex = 0; UVLayerIndex < NumUVLayers; ++UVLayerIndex)
 			{
 				FDynamicMeshUVOverlay* UVOverlay = OutputMesh.Attributes()->GetUVLayer(UVLayerIndex);
-				for (int32 TriID = 0; TriID < TriangleCount; ++TriID)
+				for (int32 SrcTriID = 0; SrcTriID < SrcTriangleCount; ++SrcTriID)
 				{
-					FVector2f UV1, UV2, UV3;
-					Adapter.GetTriUVs<FVector2f>(TriID, UVLayerIndex, UV1, UV2, UV3);
-					int32 a = UVOverlay->AppendElement(UV1);
-					int32 b = UVOverlay->AppendElement(UV2);
-					int32 c = UVOverlay->AppendElement(UV3);
-					UVOverlay->SetTriangle(TriID, FIndex3i(a, b, c));
+					const int32 DstTriID = ToDstTriID[SrcTriID];
+					if (DstTriID != FDynamicMesh3::InvalidID)
+					{ 
+						FVector2f UV1, UV2, UV3;
+						Adapter.GetTriUVs<FVector2f>(SrcTriID, UVLayerIndex, UV1, UV2, UV3);
+						int32 a = UVOverlay->AppendElement(UV1);
+						int32 b = UVOverlay->AppendElement(UV2);
+						int32 c = UVOverlay->AppendElement(UV3);
+						UVOverlay->SetTriangle(DstTriID, FIndex3i(a, b, c));
+					}
 				}
 			}
 		}
@@ -163,15 +242,19 @@ bool FStaticMeshLODResourcesToDynamicMesh::Convert(
 	{
 		OutputMesh.Attributes()->EnablePrimaryColors();
 		FDynamicMeshColorOverlay* Colors = OutputMesh.Attributes()->PrimaryColors();
-		for (int32 TriID = 0; TriID < TriangleCount; ++TriID)
+		for (int32 SrcTriID = 0; SrcTriID < SrcTriangleCount; ++SrcTriID)
 		{
-			FColor C1, C2, C3;
-			Adapter.GetTriColors(TriID, C1, C2, C3);
-			// [TODO] should we be doing RGBA conversion here?
-			int32 a = Colors->AppendElement( C1.ReinterpretAsLinear()  );
-			int32 b = Colors->AppendElement( C2.ReinterpretAsLinear() );
-			int32 c = Colors->AppendElement( C3.ReinterpretAsLinear() );
-			Colors->SetTriangle(TriID, FIndex3i(a, b, c));
+			const int32 DstTriID = ToDstTriID[SrcTriID];
+			if (DstTriID != FDynamicMesh3::InvalidID)
+			{
+				FColor C1, C2, C3;
+				Adapter.GetTriColors(SrcTriID, C1, C2, C3);
+				// [TODO] should we be doing RGBA conversion here?
+				int32 a = Colors->AppendElement( C1.ReinterpretAsLinear()  );
+				int32 b = Colors->AppendElement( C2.ReinterpretAsLinear() );
+				int32 c = Colors->AppendElement( C3.ReinterpretAsLinear() );
+				Colors->SetTriangle(DstTriID, FIndex3i(a, b, c));
+			}
 		}
 	}
 
