make static mesh lod -> dynamic mesh conversion code handle overlays more consistently, by using a shared lambda that always creates one element per vertex (rather than sometimes doing that, and sometimes creating 3 elements per triangle)

diff --git a/Engine/Source/Runtime/MeshConversionEngineTypes/Private/StaticMeshLODResourcesToDynamicMesh.cpp b/Engine/Source/Runtime/MeshConversionEngineTypes/Private/StaticMeshLODResourcesToDynamicMesh.cpp
--- a/Engine/Source/Runtime/MeshConversionEngineTypes/Private/StaticMeshLODResourcesToDynamicMesh.cpp
+++ b/Engine/Source/Runtime/MeshConversionEngineTypes/Private/StaticMeshLODResourcesToDynamicMesh.cpp
@@ -144,30 +144,36 @@ bool FStaticMeshLODResourcesToDynamicMesh::Convert(
 		}
 	}
 
+	// Helper to copy a per-vertex attribute to an overlay
+	auto CopyPerVertexOverlay = [SrcTriangleCount, &ToSrcVID, &ToDstTriID, &OutputMesh](auto* Overlay, auto GetElement)
+	{
+		const int32 DstVertexCount = ToSrcVID.Num();
+		for (int32 DstVertID = 0; DstVertID < DstVertexCount; ++DstVertID)
+		{
+			const int32 SrcVID = ToSrcVID[DstVertID];
+			auto Element = GetElement(SrcVID);
+			int32 ElemID = Overlay->AppendElement(Element);
+			check(ElemID == DstVertID);
+		}
+
+		for (int32 SrcTriID = 0; SrcTriID < SrcTriangleCount; ++SrcTriID)
+		{
+			const int32 DstTriID = ToDstTriID[SrcTriID];
+			if (DstTriID != FDynamicMesh3::InvalidID)
+			{
+				FIndex3i Tri = OutputMesh.GetTriangle(DstTriID);
+				Overlay->SetTriangle(DstTriID, FIndex3i(Tri.A, Tri.B, Tri.C));
+			}
+		}
+	};
+
 	// copy overlay normals
 	if (Adapter.HasNormals() && Options.bWantNormals)
 	{
 		FDynamicMeshNormalOverlay* Normals = OutputMesh.Attributes()->PrimaryNormals();
 		if (Normals != nullptr)
 		{
-			const int32 DstVertexCount = ToSrcVID.Num();
-			for (int32 DstVertID = 0; DstVertID < DstVertexCount; ++DstVertID)
-			{
-				const int32 SrcVID = ToSrcVID[DstVertID];
-				FVector3f N = Adapter.GetNormal(SrcVID);
-				int32 ElemID = Normals->AppendElement(N);
-				check(ElemID == DstVertID);
-			}
-
-			for (int32 SrcTriID = 0; SrcTriID < SrcTriangleCount; ++SrcTriID)
-			{
-				const int32 DstTriID = ToDstTriID[SrcTriID];
-				if (DstTriID != FDynamicMesh3::InvalidID)
-				{ 
-					FIndex3i Tri = OutputMesh.GetTriangle(DstTriID);
-					Normals->SetTriangle(DstTriID, FIndex3i(Tri.A, Tri.B, Tri.C));
-				}
-			}
+			CopyPerVertexOverlay(Normals, [&](int32 SrcVID)->FVector3f { return Adapter.GetNormal(SrcVID); });
 		}
 	}
 
@@ -178,37 +184,13 @@ bool FStaticMeshLODResourcesToDynamicMesh::Convert(
 		FDynamicMeshNormalOverlay* TangentsX = OutputMesh.Attributes()->PrimaryTangents();
 		if (TangentsX != nullptr)
 		{
-			for (int32 SrcTriID = 0; SrcTriID < SrcTriangleCount; ++SrcTriID)
-			{
-				const int32 DstTriID = ToDstTriID[SrcTriID];
-				if (DstTriID != FDynamicMesh3::InvalidID)
-				{ 
-					FVector3f T1, T2, T3;
-					Adapter.GetTriTangentsX<FVector3f>(SrcTriID, T1, T2, T3);
-					int32 a = TangentsX->AppendElement(T1);
-					int32 b = TangentsX->AppendElement(T2);
-					int32 c = TangentsX->AppendElement(T3);
-					TangentsX->SetTriangle(DstTriID, FIndex3i(a, b, c));
-				}
-			}
+			CopyPerVertexOverlay(TangentsX, [&](int32 SrcVID)->FVector3f { return Adapter.GetTangentX(SrcVID); });
 		}
 
 		FDynamicMeshNormalOverlay* TangentsY = OutputMesh.Attributes()->PrimaryBiTangents();
 		if (TangentsY != nullptr)
 		{
-			for (int32 SrcTriID = 0; SrcTriID < SrcTriangleCount; ++SrcTriID)
-			{
-				const int32 DstTriID = ToDstTriID[SrcTriID];
-				if (DstTriID != FDynamicMesh3::InvalidID)
-				{ 
-					FVector3f T1, T2, T3;
-					Adapter.GetTriTangentsY<FVector3f>(SrcTriID, T1, T2, T3);
-					int32 a = TangentsY->AppendElement(T1);
-					int32 b = TangentsY->AppendElement(T2);
-					int32 c = TangentsY->AppendElement(T3);
-					TangentsY->SetTriangle(DstTriID, FIndex3i(a, b, c));
-				}
-			}
+			CopyPerVertexOverlay(TangentsY, [&](int32 SrcVID)->FVector3f { return Adapter.GetTangentY(SrcVID); });
 		}
 	}
 
@@ -222,19 +204,7 @@ bool FStaticMeshLODResourcesToDynamicMesh::Convert(
 			for (int32 UVLayerIndex = 0; UVLayerIndex < NumUVLayers; ++UVLayerIndex)
 			{
 				FDynamicMeshUVOverlay* UVOverlay = OutputMesh.Attributes()->GetUVLayer(UVLayerIndex);
-				for (int32 SrcTriID = 0; SrcTriID < SrcTriangleCount; ++SrcTriID)
-				{
-					const int32 DstTriID = ToDstTriID[SrcTriID];
-					if (DstTriID != FDynamicMesh3::InvalidID)
-					{ 
-						FVector2f UV1, UV2, UV3;
-						Adapter.GetTriUVs<FVector2f>(SrcTriID, UVLayerIndex, UV1, UV2, UV3);
-						int32 a = UVOverlay->AppendElement(UV1);
-						int32 b = UVOverlay->AppendElement(UV2);
-						int32 c = UVOverlay->AppendElement(UV3);
-						UVOverlay->SetTriangle(DstTriID, FIndex3i(a, b, c));
-					}
-				}
+				CopyPerVertexOverlay(UVOverlay, [&](int32 SrcVID)->FVector2f { return Adapter.GetUV(SrcVID, UVLayerIndex); });
 			}
 		}
 	}
@@ -245,24 +215,7 @@ bool FStaticMeshLODResourcesToDynamicMesh::Convert(
 		OutputMesh.Attributes()->EnablePrimaryColors();
 		FDynamicMeshColorOverlay* Colors = OutputMesh.Attributes()->PrimaryColors();
 
-		const int32 DstVertexCount = ToSrcVID.Num();
-		for (int32 DstVertID = 0; DstVertID < DstVertexCount; ++DstVertID)
-		{
-			const int32 SrcVID = ToSrcVID[DstVertID];
-			FColor C = GetVertexColorFromLODVertexIndex(SrcVID);
-			int32 ElemID = Colors->AppendElement(C.ReinterpretAsLinear());
-			check(ElemID == DstVertID);
-		}
-
-		for (int32 SrcTriID = 0; SrcTriID < SrcTriangleCount; ++SrcTriID)
-		{
-			const int32 DstTriID = ToDstTriID[SrcTriID];
-			if (DstTriID != FDynamicMesh3::InvalidID)
-			{
-				FIndex3i Tri = OutputMesh.GetTriangle(DstTriID);
-				Colors->SetTriangle(DstTriID, FIndex3i(Tri.A, Tri.B, Tri.C));
-			}
-		}
+		CopyPerVertexOverlay(Colors, [&](int32 SrcVID)->FVector4f { return GetVertexColorFromLODVertexIndex(SrcVID).ReinterpretAsLinear(); });
 	}
 
 	return true;
