Fix handling of updating a geometry collection from a dynamic mesh when the dynamic mesh doesn't have all attributes

diff --git a/Engine/Source/Runtime/MeshConversionEngineTypes/Private/GeometryCollectionToDynamicMesh.cpp b/Engine/Source/Runtime/MeshConversionEngineTypes/Private/GeometryCollectionToDynamicMesh.cpp
--- a/Engine/Source/Runtime/MeshConversionEngineTypes/Private/GeometryCollectionToDynamicMesh.cpp
+++ b/Engine/Source/Runtime/MeshConversionEngineTypes/Private/GeometryCollectionToDynamicMesh.cpp
@@ -305,19 +305,25 @@ namespace
 
 		TUniqueMeshVertex()
 		{
-			ColorEID = -1;
-			NormalEIDs = { -1,-1,-1 };
-			for (int32 UVIdx = 0; UVIdx < MaxUVLayers; ++UVIdx)
-			{
-				UVsEIDs[UVIdx] = -1;
-			}
+			InitEmpty();
 		}
 		TUniqueMeshVertex(const FDynamicMesh3& Mesh, int32 TID, int32 SubIdx)
 		{
-			NormalEIDs[0] = Mesh.Attributes()->PrimaryNormals()->GetTriangle(TID)[SubIdx];
-			NormalEIDs[1] = Mesh.Attributes()->PrimaryTangents()->GetTriangle(TID)[SubIdx];
-			NormalEIDs[2] = Mesh.Attributes()->PrimaryBiTangents()->GetTriangle(TID)[SubIdx];
-			ColorEID = Mesh.Attributes()->PrimaryColors()->GetTriangle(TID)[SubIdx];
+			if (!Mesh.HasAttributes())
+			{
+				InitEmpty();
+				return;
+			}
+			const int32 NumNormalLayers = FMath::Min(3, Mesh.Attributes()->NumNormalLayers());
+			for (int32 NormalLayerIdx = 0; NormalLayerIdx < NumNormalLayers; ++NormalLayerIdx)
+			{
+				NormalEIDs[NormalLayerIdx] = Mesh.Attributes()->GetNormalLayer(NormalLayerIdx)->GetTriangle(TID)[SubIdx];
+			}
+			for (int32 NormalLayerIdx = NumNormalLayers; NormalLayerIdx < 3; ++NormalLayerIdx)
+			{
+				NormalEIDs[NormalLayerIdx] = -1;
+			}
+			ColorEID = Mesh.Attributes()->PrimaryColors() ? Mesh.Attributes()->PrimaryColors()->GetTriangle(TID)[SubIdx] : -1;
 			int32 NumUVs = FMath::Min(MaxUVLayers, Mesh.Attributes()->NumUVLayers()); // max 8 UV layers supported
 			for (int32 UVIdx = 0; UVIdx < NumUVs; ++UVIdx)
 			{
@@ -329,6 +335,16 @@ namespace
 			}
 		}
 
+		inline void InitEmpty()
+		{
+			ColorEID = -1;
+			NormalEIDs = { -1,-1,-1 };
+			for (int32 UVIdx = 0; UVIdx < MaxUVLayers; ++UVIdx)
+			{
+				UVsEIDs[UVIdx] = -1;
+			}
+		}
+
 		bool operator==(const TUniqueMeshVertex& Other) const
 		{
 
