Fix geometry collection mesh append to not create invalid parent hierarchy when appending to an empty collection + handle an invalid requested parent index

diff --git a/Engine/Source/Runtime/MeshConversionEngineTypes/Private/GeometryCollectionToDynamicMesh.cpp b/Engine/Source/Runtime/MeshConversionEngineTypes/Private/GeometryCollectionToDynamicMesh.cpp
--- a/Engine/Source/Runtime/MeshConversionEngineTypes/Private/GeometryCollectionToDynamicMesh.cpp
+++ b/Engine/Source/Runtime/MeshConversionEngineTypes/Private/GeometryCollectionToDynamicMesh.cpp
@@ -499,10 +499,14 @@ int32 FGeometryCollectionToDynamicMeshes::AppendMeshToCollection(FGeometryCollec
 	Collection.TransformIndex[GeometryIdx] = TransformIdx;
 	Collection.TransformToGeometryIndex[TransformIdx] = GeometryIdx;
 	int32 TransformParent = Options.NewMeshParentIndex;
-	// If parent is INDEX_NONE, search for a root to use as parent
+	if (TransformParent >= TransformIdx)
+	{
+		TransformParent = INDEX_NONE;
+	}
+	// If parent is INDEX_NONE, search for a root (other than self) to use as parent
 	if (TransformParent == INDEX_NONE)
 	{
-		for (int32 Idx = 0; Idx < Collection.Parent.Num(); ++Idx)
+		for (int32 Idx = 0; Idx < TransformIdx; ++Idx)
 		{
 			if (Collection.Parent[Idx] == -1)
 			{
