make sure the geometry collection mesh append logic's automatic root selection logic only considers cluster nodes

diff --git a/Engine/Source/Runtime/MeshConversionEngineTypes/Private/GeometryCollectionToDynamicMesh.cpp b/Engine/Source/Runtime/MeshConversionEngineTypes/Private/GeometryCollectionToDynamicMesh.cpp
--- a/Engine/Source/Runtime/MeshConversionEngineTypes/Private/GeometryCollectionToDynamicMesh.cpp
+++ b/Engine/Source/Runtime/MeshConversionEngineTypes/Private/GeometryCollectionToDynamicMesh.cpp
@@ -503,12 +503,12 @@ int32 FGeometryCollectionToDynamicMeshes::AppendMeshToCollection(FGeometryCollec
 	{
 		TransformParent = INDEX_NONE;
 	}
-	// If parent is INDEX_NONE, search for a root (other than self) to use as parent
+	// If parent is INDEX_NONE, search for a cluster root (other than self) to use as parent
 	if (TransformParent == INDEX_NONE)
 	{
 		for (int32 Idx = 0; Idx < TransformIdx; ++Idx)
 		{
-			if (Collection.Parent[Idx] == -1)
+			if (Collection.Parent[Idx] == -1 && Collection.SimulationType[Idx] == FGeometryCollection::ESimulationTypes::FST_Clustered)
 			{
 				TransformParent = Idx;
 			}
