fix dynamic mesh to geometry collection conversion for the non-compact mesh case

diff --git a/Engine/Source/Runtime/MeshConversionEngineTypes/Private/GeometryCollectionToDynamicMesh.cpp b/Engine/Source/Runtime/MeshConversionEngineTypes/Private/GeometryCollectionToDynamicMesh.cpp
--- a/Engine/Source/Runtime/MeshConversionEngineTypes/Private/GeometryCollectionToDynamicMesh.cpp
+++ b/Engine/Source/Runtime/MeshConversionEngineTypes/Private/GeometryCollectionToDynamicMesh.cpp
@@ -705,7 +705,13 @@ bool FGeometryCollectionToDynamicMeshes::UpdateCollection(const FTransform& From
 			TUniqueMeshVertex<> Unique(Mesh, TID, SubIdx);
 			int32 MeshVID = Tri[SubIdx];
 			int32 StartIdx = VertexIDToIdxStart[MeshVID];
-			int32 EndIdx = MeshVID + 1 < VertexIDToIdxStart.Num() ? VertexIDToIdxStart[MeshVID + 1] : PerIdxElements.Num();
+			int32 NextMeshVID = MeshVID + 1;
+			// Search for a valid NextMeshVID (in case mesh is not compact)
+			while (NextMeshVID < VertexIDToIdxStart.Num() && VertexIDToIdxStart[NextMeshVID] == INDEX_NONE)
+			{
+				NextMeshVID++;
+			}
+			int32 EndIdx = NextMeshVID < VertexIDToIdxStart.Num() ? VertexIDToIdxStart[NextMeshVID] : PerIdxElements.Num();
 			int32 FoundIdx = -1;
 			for (int32 TestIdx = StartIdx; TestIdx < EndIdx; ++TestIdx)
 			{
