Fixed incorrect transforms on GeometryCollectionToDynamicMesh converter, and added more comments to document the expected coordinate spaces for these meshes

diff --git a/Engine/Source/Runtime/MeshConversionEngineTypes/Private/GeometryCollectionToDynamicMesh.cpp b/Engine/Source/Runtime/MeshConversionEngineTypes/Private/GeometryCollectionToDynamicMesh.cpp
--- a/Engine/Source/Runtime/MeshConversionEngineTypes/Private/GeometryCollectionToDynamicMesh.cpp
+++ b/Engine/Source/Runtime/MeshConversionEngineTypes/Private/GeometryCollectionToDynamicMesh.cpp
@@ -281,12 +281,12 @@ bool FGeometryCollectionToDynamicMeshes::InitHelper(const FManagedArrayCollectio
 // Convert all geometry in the given Collection to dynamic meshes
 bool FGeometryCollectionToDynamicMeshes::Init(const FManagedArrayCollection& Collection, const FToMeshOptions& Options)
 {
-	return InitHelper(Collection, true, GetCollectionTransforms(Collection), true, TConstArrayView<int32>(), Options);
+	return InitHelper(Collection, false, GetCollectionTransforms(Collection), true, TConstArrayView<int32>(), Options);
 }
 // Convert geometry at selected transform indices in the given Collection to dynamic meshes
 bool FGeometryCollectionToDynamicMeshes::InitFromTransformSelection(const FManagedArrayCollection& Collection, TConstArrayView<int32> TransformIndices, const FToMeshOptions& Options)
 {
-	return InitHelper(Collection, true, GetCollectionTransforms(Collection), false, TransformIndices, Options);
+	return InitHelper(Collection, false, GetCollectionTransforms(Collection), false, TransformIndices, Options);
 }
 
 namespace
diff --git a/Engine/Source/Runtime/MeshConversionEngineTypes/Public/GeometryCollectionToDynamicMesh.h b/Engine/Source/Runtime/MeshConversionEngineTypes/Public/GeometryCollectionToDynamicMesh.h
--- a/Engine/Source/Runtime/MeshConversionEngineTypes/Public/GeometryCollectionToDynamicMesh.h
+++ b/Engine/Source/Runtime/MeshConversionEngineTypes/Public/GeometryCollectionToDynamicMesh.h
@@ -17,8 +17,11 @@ struct FGeometryCollectionToDynamicMeshes
 {
 	struct FMeshInfo
 	{
+		// Transform Index of the bone represented
 		int32 TransformIndex = -1;
+		// Mesh holding the geometry at (or under) a given bone, in a common coordinate space
 		TUniquePtr<FDynamicMesh3> Mesh;
+		// Stores the Transform that has *already* been applied to the Mesh vertices, which took it from local bone space to a common space
 		FTransform Transform = FTransform::Identity;
 	};
 
@@ -81,6 +84,7 @@ struct FGeometryCollectionToDynamicMeshes
 private:
 	
 	// private helper method with shared implementation for all Init* methods
+	// @param bTransformInComponentSpace If false, the Transforms array holds transforms that are relative to the corresponding bone parent. If true, the Transforms are all relative to a common coordinate space.
 	bool InitHelper(const FManagedArrayCollection& Collection, bool bTransformInComponentSpace, TConstArrayView<FTransform3f> Transforms, bool bAllTransforms, TConstArrayView<int32> TransformIndices, const FToMeshOptions& Options);
 	static TConstArrayView<FTransform3f> GetCollectionTransforms(const FManagedArrayCollection& Collection);
 
